#ue4
#ios
* missing entitlements file

diff --git a/Engine/Build/IOS/UE4Game.entitlements b/Engine/Build/IOS/UE4Game.entitlements
new file mode 100644
--- /dev/null
+++ b/Engine/Build/IOS/UE4Game.entitlements
@@ -0,0 +1 @@
+<plist><dict><key>get-task-allow</key><true/></dict></plist>
\ No newline at end of file
