#ios Fixed ~/.lldbinit, to make breakpoints work in Xcode

diff --git a/Engine/Build/IOS/XcodeSupportFiles/prepackage.sh b/Engine/Build/IOS/XcodeSupportFiles/prepackage.sh
--- a/Engine/Build/IOS/XcodeSupportFiles/prepackage.sh
+++ b/Engine/Build/IOS/XcodeSupportFiles/prepackage.sh
@@ -11,7 +11,8 @@ else
 fi
 
 echo "Checking for inline breakpoint strategy"
-if [ -n "ls | grep 'strategy' ~/.lldbinit" ]
+lldbinitfile=~/.lldbinit
+if grep -q "strategy" $lldbinitfile
 then
   echo "Strategy found... inline breakpoints already enabled"
 else
