Fix for relative path for plist file
#ue4
#ios

diff --git a/Engine/Build/IOS/XcodeSupportFiles/prepackage.sh b/Engine/Build/IOS/XcodeSupportFiles/prepackage.sh
--- a/Engine/Build/IOS/XcodeSupportFiles/prepackage.sh
+++ b/Engine/Build/IOS/XcodeSupportFiles/prepackage.sh
@@ -50,12 +50,12 @@ then
   popd > /dev/null
 fi
 
-if [ -e ../../../../Engine/Binaries/IOS/XcodeSupportFiles/$1-Info.plist ]
+if [ -e $1-Info.plist ]
 then
   if [ -d ../../../../$1/Intermediate/IOS ]
   then
-    cp -R ../../../../Engine/Binaries/IOS/XcodeSupportFiles/$1-Info.plist ../../../../$1/Intermediate/IOS/$1-Info.plist
+    cp -R $1-Info.plist ../../../../$1/Intermediate/IOS/$1-Info.plist
   else
-    cp -R ../../../../Engine/Binaries/IOS/XcodeSupportFiles/$1-Info.plist ../../../../Engine/Binaries/IOS/Engine/Intermediate/IOS/$1-Info.plist
+    cp -R $1-Info.plist ../../../../Engine/Binaries/IOS/Engine/Intermediate/IOS/$1-Info.plist
   fi
 fi
\ No newline at end of file
