Fix to make the directory before copying the plist for iOS PC toolchain
#ue4
#ios

diff --git a/Engine/Build/IOS/XcodeSupportFiles/prepackage.sh b/Engine/Build/IOS/XcodeSupportFiles/prepackage.sh
--- a/Engine/Build/IOS/XcodeSupportFiles/prepackage.sh
+++ b/Engine/Build/IOS/XcodeSupportFiles/prepackage.sh
@@ -56,6 +56,7 @@ then
   then
     cp -R $1-Info.plist ../../../../$1/Intermediate/IOS/$1-Info.plist
   else
+    mkdir -p ../../../../Engine/Binaries/IOS/Engine/Intermediate/IOS/
     cp -R $1-Info.plist ../../../../Engine/Binaries/IOS/Engine/Intermediate/IOS/$1-Info.plist
   fi
 fi
\ No newline at end of file
