fix for device family error when modifying iOS project settings
UE-4797
#ios

diff --git a/Engine/Build/IOS/UE4Game-Info.plist b/Engine/Build/IOS/UE4Game-Info.plist
--- a/Engine/Build/IOS/UE4Game-Info.plist
+++ b/Engine/Build/IOS/UE4Game-Info.plist
@@ -115,5 +115,10 @@
   </array>
   <key>MinimumOSVersion</key>
   <string>6.0</string>
+  <key>UIDeviceFamily</key>
+  <array>
+    <integer>1</integer>
+    <integer>2</integer>
+  </array>
 </dict>
 </plist>
