fix for incorrect location of the plist when the game directory is not the same name as the game
UE-20492
#ios

diff --git a/Engine/Build/IOS/XcodeSupportFiles/prepackage.sh b/Engine/Build/IOS/XcodeSupportFiles/prepackage.sh
--- a/Engine/Build/IOS/XcodeSupportFiles/prepackage.sh
+++ b/Engine/Build/IOS/XcodeSupportFiles/prepackage.sh
@@ -20,9 +20,9 @@ else
   echo "Strategy not found -> enabling all inline breakpoints"
 fi
 
-if [ -e ../../../../$1/Binaries/IOS/$1-$2-$3$4.app.dSYM.zip.datecheck ]
+if [ -e ../../../Binaries/IOS/$1-$2-$3$4.app.dSYM.zip.datecheck ]
 then
-  pushd ../../../../$1/Binaries/IOS > /dev/null
+  pushd ../../../Binaries/IOS > /dev/null
   if [ $1-$2-$3$4.app.dSYM.zip.datecheck -nt $1-$2-$3.app.dSYM.zip ]
   then
     echo "Unzipping new .dSYM..."
@@ -36,33 +36,33 @@ else
   echo $1"-"$2"-"$3""$4".app.dSYM not found..."
 fi
 
-if [ -e ../../../../$1/Binaries/IOS/$1-$2-$3$4.app.dSYM ]
+if [ -e ../../../Binaries/IOS/$1-$2-$3$4.app.dSYM ]
 then
-  pushd ../../../../$1/Binaries/IOS > /dev/null
+  pushd ../../../Binaries/IOS > /dev/null
   cp -R $1-$2-$3$4.app.dSYM Payload/$1.app.dSYM
   popd > /dev/null
 fi
 
-if [ -e ../../../../$1/Binaries/IOS/$1$4.app.dSYM ]
+if [ -e ../../../Binaries/IOS/$1$4.app.dSYM ]
 then
-  pushd ../../../../$1/Binaries/IOS > /dev/null
+  pushd ../../../Binaries/IOS > /dev/null
   cp -R $1$4.app.dSYM Payload/$1$4.app.dSYM
   popd > /dev/null
 fi
 
 if [ -e $1-Info.plist ]
 then
-  if [ -d ../../../../$1/Intermediate/IOS ]
+  if [ -d ../../../Intermediate/IOS ]
   then
-    cp -R $1-Info.plist ../../../../$1/Intermediate/IOS/$1-Info.plist
+    cp -R $1-Info.plist ../../../Intermediate/IOS/$1-Info.plist
   else
 	if [ "$1" == "UE4Game" ]
 	then
       mkdir -p ../../../../Engine/Intermediate/IOS/
       cp -R $1-Info.plist ../../../../Engine/Intermediate/IOS/$1-Info.plist
 	else 
-      mkdir -p ../../../../$1/Intermediate/IOS/
-      cp -R $1-Info.plist ../../../../$1/Intermediate/IOS/$1-Info.plist
+      mkdir -p ../../../Intermediate/IOS/
+      cp -R $1-Info.plist ../../../Intermediate/IOS/$1-Info.plist
 	fi
   fi
 fi
\ No newline at end of file
