Fixed case on IOS files

diff --git a/Engine/Build/IOS/UE4CmdLineRun/UE4CmdLineRun.m b/Engine/Build/IOS/UE4CmdLineRun/UE4CmdLineRun.m
new file mode 100644
--- /dev/null
+++ b/Engine/Build/IOS/UE4CmdLineRun/UE4CmdLineRun.m
@@ -0,0 +1,33 @@
+// Copyright 2013-2017 Epic Games, Inc. All Rights Reserved.
+
+#import <UIKit/UIKit.h>
+#import <Foundation/Foundation.h>
+#import <XCTest/XCTest.h>
+
+@interface UE4CmdLineRun : XCTestCase
+
+@end
+
+@implementation UE4CmdLineRun
+
+- (void)setUp
+{
+    [super setUp];
+    // Put setup code here. This method is called before the invocation of each test method in the class.
+}
+
+- (void)tearDown
+{
+    // Put teardown code here. This method is called after the invocation of each test method in the class.
+    [super tearDown];
+}
+
+- (void)testCmdLineRun
+{
+	// by running the runloop from the test, the full game will be running as if launched normalled
+	// the unit test is run from the startup of the normal application path, this lets it continue on
+	// since the system will quit the app as soon as all tests are complete
+	[[NSRunLoop currentRunLoop] run];
+}
+
+@end
diff --git a/Engine/Build/IOS/XcodeSupportFiles/prepackage.sh b/Engine/Build/IOS/XcodeSupportFiles/prepackage.sh
new file mode 100755
--- /dev/null
+++ b/Engine/Build/IOS/XcodeSupportFiles/prepackage.sh
@@ -0,0 +1,86 @@
+#!/bin/sh
+
+echo "Checking project"
+if [ project.pbxproj.datecheck -nt ../UE4_FromPC.xcodeproj/project.pbxproj ]
+then
+  echo "Copying project..."
+  mkdir -p ../UE4_FromPC.xcodeproj > /dev/null
+  cp project.pbxproj.datecheck ../UE4_FromPC.xcodeproj/project.pbxproj
+else
+  echo "Project file is up to date."
+fi
+
+echo "Checking for inline breakpoint strategy"
+lldbinitfile=~/.lldbinit
+if grep -q "strategy" $lldbinitfile
+then
+  echo "Strategy found... inline breakpoints already enabled"
+else
+  echo "settings set target.inline-breakpoint-strategy always" >> ~/.lldbinit
+  echo "Strategy not found -> enabling all inline breakpoints"
+fi
+
+if [ -e ../../../Binaries/IOS/$1-$2-$3$4.app.dSYM.zip.datecheck ]
+then
+  pushd ../../../Binaries/IOS > /dev/null
+  if [ $1-$2-$3$4.app.dSYM.zip.datecheck -nt $1-$2-$3.app.dSYM.zip ]
+  then
+    echo "Unzipping new .dSYM..."
+    mv $1-$2-$3$4.app.dSYM.zip.datecheck $1-$2-$3$4.app.dSYM.zip
+    unzip -o $1-$2-$3$4.app.dSYM.zip
+  else
+    echo ".dSYM is up to date."
+  fi
+  popd > /dev/null
+else
+  echo $1"-"$2"-"$3""$4".app.dSYM not found..."
+fi
+
+if [ -e ../../../Binaries/IOS/$1-$2-$3$4.app.dSYM ]
+then
+  pushd ../../../Binaries/IOS > /dev/null
+  cp -R $1-$2-$3$4.app.dSYM Payload/$1.app.dSYM
+  popd > /dev/null
+fi
+
+if [ -e ../../../Binaries/IOS/$1$4.app.dSYM ]
+then
+  pushd ../../../Binaries/IOS > /dev/null
+  cp -R $1$4.app.dSYM Payload/$1$4.app.dSYM
+  popd > /dev/null
+fi
+
+if [ -e $1-Info.plist ]
+then
+  if [ -d ../../../Intermediate/IOS ]
+  then
+    cp -R $1-Info.plist ../../../Intermediate/IOS/$1-Info.plist
+  else
+	if [ "$1" == "UE4Game" ]
+	then
+      mkdir -p ../../../../Engine/Intermediate/IOS/
+      cp -R $1-Info.plist ../../../../Engine/Intermediate/IOS/$1-Info.plist
+	else 
+      mkdir -p ../../../Intermediate/IOS/
+      cp -R $1-Info.plist ../../../Intermediate/IOS/$1-Info.plist
+	fi
+  fi
+fi
+
+
+if [ -e $1.entitlements ]
+then
+  if [ -d ../../../Intermediate/IOS ]
+  then
+    cp -R $1.entitlements ../../../Intermediate/IOS/$1.entitlements
+  else
+	if [ "$1" == "UE4Game" ]
+	then
+      mkdir -p ../../../../Engine/Intermediate/IOS/
+      cp -R $1.entitlements ../../../../Engine/Intermediate/IOS/$1.entitlements
+	else 
+      mkdir -p ../../../Intermediate/IOS/
+      cp -R $1.entitlements ../../../Intermediate/IOS/$1.entitlements
+	fi
+  fi
+fi
\ No newline at end of file
