Add copyright boilerplate to beginning of each file
Change line endings in IOSToolChain.cs to prevent build warnings about XML comments

diff --git a/Engine/Build/IOS/Resources/FrameworkWrapper/FrameworkWrapper/IOS/AppDelegate.h b/Engine/Build/IOS/Resources/FrameworkWrapper/FrameworkWrapper/IOS/AppDelegate.h
--- a/Engine/Build/IOS/Resources/FrameworkWrapper/FrameworkWrapper/IOS/AppDelegate.h
+++ b/Engine/Build/IOS/Resources/FrameworkWrapper/FrameworkWrapper/IOS/AppDelegate.h
@@ -1,9 +1,9 @@
+// Copyright 1998-2019 Epic Games, Inc. All Rights Reserved.
 //
 //  AppDelegate.h
 //  FrameworkWrapper
 //
 //  Created by Ryan West on 6/13/19.
-//  Copyright © 2019 Ryan West. All rights reserved.
 //
 
 #import <UIKit/UIKit.h>
diff --git a/Engine/Build/IOS/Resources/FrameworkWrapper/FrameworkWrapper/IOS/AppDelegate.mm b/Engine/Build/IOS/Resources/FrameworkWrapper/FrameworkWrapper/IOS/AppDelegate.mm
--- a/Engine/Build/IOS/Resources/FrameworkWrapper/FrameworkWrapper/IOS/AppDelegate.mm
+++ b/Engine/Build/IOS/Resources/FrameworkWrapper/FrameworkWrapper/IOS/AppDelegate.mm
@@ -1,9 +1,9 @@
+// Copyright 1998-2019 Epic Games, Inc. All Rights Reserved.
 //
 //  AppDelegate.m
 //  FrameworkWrapper
 //
 //  Created by Ryan West on 6/13/19.
-//  Copyright © 2019 Ryan West. All rights reserved.
 //
 
 #import "AppDelegate.h"
diff --git a/Engine/Build/IOS/Resources/FrameworkWrapper/FrameworkWrapper/IOS/main.mm b/Engine/Build/IOS/Resources/FrameworkWrapper/FrameworkWrapper/IOS/main.mm
--- a/Engine/Build/IOS/Resources/FrameworkWrapper/FrameworkWrapper/IOS/main.mm
+++ b/Engine/Build/IOS/Resources/FrameworkWrapper/FrameworkWrapper/IOS/main.mm
@@ -1,9 +1,9 @@
+// Copyright 1998-2019 Epic Games, Inc. All Rights Reserved.
 //
 //  main.m
 //  FrameworkWrapper
 //
 //  Created by Ryan West on 6/13/19.
-//  Copyright © 2019 Ryan West. All rights reserved.
 //
 
 
diff --git a/Engine/Build/IOS/Resources/FrameworkWrapper/FrameworkWrapper/ViewControllers/ViewController.h b/Engine/Build/IOS/Resources/FrameworkWrapper/FrameworkWrapper/ViewControllers/ViewController.h
--- a/Engine/Build/IOS/Resources/FrameworkWrapper/FrameworkWrapper/ViewControllers/ViewController.h
+++ b/Engine/Build/IOS/Resources/FrameworkWrapper/FrameworkWrapper/ViewControllers/ViewController.h
@@ -1,9 +1,9 @@
+// Copyright 1998-2019 Epic Games, Inc. All Rights Reserved.
 //
 //  ViewController.h
 //  FrameworkWrapper
 //
 //  Created by Ryan West on 7/29/19.
-//  Copyright © 2019 Epic Games. All rights reserved.
 //
 
 #ifndef ViewController_h
diff --git a/Engine/Build/IOS/Resources/FrameworkWrapper/FrameworkWrapper/ViewControllers/ViewController.m b/Engine/Build/IOS/Resources/FrameworkWrapper/FrameworkWrapper/ViewControllers/ViewController.m
--- a/Engine/Build/IOS/Resources/FrameworkWrapper/FrameworkWrapper/ViewControllers/ViewController.m
+++ b/Engine/Build/IOS/Resources/FrameworkWrapper/FrameworkWrapper/ViewControllers/ViewController.m
@@ -1,9 +1,9 @@
+// Copyright 1998-2019 Epic Games, Inc. All Rights Reserved.
 //
 //  ViewController.m
 //  FrameworkWrapper
 //
 //  Created by Ryan West on 7/29/19.
-//  Copyright © 2019 Epic Games. All rights reserved.
 //
 
 #import "ViewController.h"
diff --git a/Engine/Build/IOS/Resources/FrameworkWrapper/FrameworkWrapper/Views/UnrealView.h b/Engine/Build/IOS/Resources/FrameworkWrapper/FrameworkWrapper/Views/UnrealView.h
--- a/Engine/Build/IOS/Resources/FrameworkWrapper/FrameworkWrapper/Views/UnrealView.h
+++ b/Engine/Build/IOS/Resources/FrameworkWrapper/FrameworkWrapper/Views/UnrealView.h
@@ -1,9 +1,9 @@
+// Copyright 1998-2019 Epic Games, Inc. All Rights Reserved.
 //
 //  UnrealWrapper.h
 //  FrameworkWrapper
 //
 //  Created by Andrew Grant on 9/4/18.
-//  Copyright © 2018 Epic Games, Inc. All rights reserved.
 //
 
 #if defined(__aarch64__)
diff --git a/Engine/Build/IOS/Resources/FrameworkWrapper/FrameworkWrapper/Views/UnrealView.mm b/Engine/Build/IOS/Resources/FrameworkWrapper/FrameworkWrapper/Views/UnrealView.mm
--- a/Engine/Build/IOS/Resources/FrameworkWrapper/FrameworkWrapper/Views/UnrealView.mm
+++ b/Engine/Build/IOS/Resources/FrameworkWrapper/FrameworkWrapper/Views/UnrealView.mm
@@ -1,9 +1,9 @@
+// Copyright 1998-2019 Epic Games, Inc. All Rights Reserved.
 //
 //  UnrealWrapper.m
 //  FrameworkWrapper
 //
 //  Created by Andrew Grant on 9/4/18.
-//  Copyright © 2018 Epic Games, Inc. All rights reserved.
 //
 
 
