Fix UE4 references in UECmdLineRun

diff --git a/Engine/Build/IOS/UECmdLineRun/UECmdLineRun.m b/Engine/Build/IOS/UECmdLineRun/UECmdLineRun.m
--- a/Engine/Build/IOS/UECmdLineRun/UECmdLineRun.m
+++ b/Engine/Build/IOS/UECmdLineRun/UECmdLineRun.m
@@ -4,11 +4,11 @@
 #import <Foundation/Foundation.h>
 #import <XCTest/XCTest.h>
 
-@interface UE4CmdLineRun : XCTestCase
+@interface UECmdLineRun : XCTestCase
 
 @end
 
-@implementation UE4CmdLineRun
+@implementation UECmdLineRun
 
 - (void)setUp
 {
